RC OSC config section visibility

diff --git a/Engine/Plugins/VirtualProduction/RemoteControlProtocolOSC/Source/RemoteControlProtocolOSC/Private/RemoteControlProtocolOSCSettings.h b/Engine/Plugins/VirtualProduction/RemoteControlProtocolOSC/Source/RemoteControlProtocolOSC/Private/RemoteControlProtocolOSCSettings.h
--- a/Engine/Plugins/VirtualProduction/RemoteControlProtocolOSC/Source/RemoteControlProtocolOSC/Private/RemoteControlProtocolOSCSettings.h
+++ b/Engine/Plugins/VirtualProduction/RemoteControlProtocolOSC/Source/RemoteControlProtocolOSC/Private/RemoteControlProtocolOSCSettings.h
@@ -28,7 +28,7 @@ public:
 	* 
 	* The format is IP_ADDRESS:PORT_NUMBER.
 	*/
-	UPROPERTY(EditAnywhere, category = OSC);
+	UPROPERTY(EditAnywhere, Category = OSC);
 	FString ServerAddress = "127.0.0.1:8001";
 
 	/** Running OSC server instance */
@@ -38,7 +38,7 @@ public:
 /**
  * OSC Remote Control settings
  */
-UCLASS(config = Engine, defaultconfig)
+UCLASS(Config = Engine, DefaultConfig, meta = (Keywords = "OSC"))
 class URemoteControlProtocolOSCSettings : public UObject
 {
 	GENERATED_BODY()
@@ -55,6 +55,6 @@ public:
 
 public:
 	/** OSC server pair of server ip and server port */
-	UPROPERTY(Config, EditAnywhere, category = Mapping)
+	UPROPERTY(Config, EditAnywhere, Category = OSC)
 	TArray<FRemoteControlOSCServerSettings> ServersSettings = { FRemoteControlOSCServerSettings() };
 };
