Handle invalid or null OSC protocol entity.

diff --git a/Engine/Plugins/VirtualProduction/RemoteControlProtocolOSC/Source/RemoteControlProtocolOSC/Private/RemoteControlProtocolOSC.cpp b/Engine/Plugins/VirtualProduction/RemoteControlProtocolOSC/Source/RemoteControlProtocolOSC/Private/RemoteControlProtocolOSC.cpp
--- a/Engine/Plugins/VirtualProduction/RemoteControlProtocolOSC/Source/RemoteControlProtocolOSC/Private/RemoteControlProtocolOSC.cpp
+++ b/Engine/Plugins/VirtualProduction/RemoteControlProtocolOSC/Source/RemoteControlProtocolOSC/Private/RemoteControlProtocolOSC.cpp
@@ -15,6 +15,8 @@
 
 #define LOCTEXT_NAMESPACE "RemoteControlProtocolOSC"
 
+DEFINE_LOG_CATEGORY_STATIC(LogRemoteControlOSC, Log, All);
+
 const FName FRemoteControlProtocolOSC::ProtocolName = TEXT("OSC");
 
 #if WITH_EDITOR
@@ -52,10 +54,15 @@ void FRemoteControlProtocolOSC::Bind(FRemoteControlProtocolEntityPtr InRemoteCon
 		return;
 	}
 
-	FRemoteControlOSCProtocolEntity* OSCProtocolEntity = InRemoteControlProtocolEntityPtr->CastChecked<FRemoteControlOSCProtocolEntity>();
-
-	TArray<FRemoteControlProtocolEntityWeakPtr>& BindingsEntitiesPtr = Bindings.FindOrAdd(OSCProtocolEntity->PathName);
-	BindingsEntitiesPtr.Add(MoveTemp(InRemoteControlProtocolEntityPtr));
+	if (FRemoteControlOSCProtocolEntity* OSCProtocolEntity = InRemoteControlProtocolEntityPtr->Cast<FRemoteControlOSCProtocolEntity>())
+	{
+		TArray<FRemoteControlProtocolEntityWeakPtr>& BindingsEntitiesPtr = Bindings.FindOrAdd(OSCProtocolEntity->PathName);
+		BindingsEntitiesPtr.Add(MoveTemp(InRemoteControlProtocolEntityPtr));
+	}
+	else
+	{
+		UE_LOG(LogRemoteControlOSC, Warning, TEXT("Invalid Remote Control Protocol Entity."))
+	}
 }
 
 void FRemoteControlProtocolOSC::Unbind(FRemoteControlProtocolEntityPtr InRemoteControlProtocolEntityPtr)
@@ -65,10 +72,16 @@ void FRemoteControlProtocolOSC::Unbind(FRemoteControlProtocolEntityPtr InRemoteC
 		return;
 	}
 
-	FRemoteControlOSCProtocolEntity* OSCProtocolEntity = InRemoteControlProtocolEntityPtr->CastChecked<FRemoteControlOSCProtocolEntity>();	
-	for (TPair<FName, TArray<FRemoteControlProtocolEntityWeakPtr>>& BindingsPair : Bindings)
+	if (FRemoteControlOSCProtocolEntity* OSCProtocolEntity = InRemoteControlProtocolEntityPtr->Cast<FRemoteControlOSCProtocolEntity>())
+	{
+		for (TPair<FName, TArray<FRemoteControlProtocolEntityWeakPtr>>& BindingsPair : Bindings)
+		{
+			BindingsPair.Value.RemoveAllSwap(CreateProtocolComparator(OSCProtocolEntity->GetPropertyId()));
+		}
+	}
+	else
 	{
-		BindingsPair.Value.RemoveAllSwap(CreateProtocolComparator(OSCProtocolEntity->GetPropertyId()));
+		UE_LOG(LogRemoteControlOSC, Warning, TEXT("Invalid Remote Control Protocol Entity."))
 	}
 }
 
@@ -146,8 +159,10 @@ void FRemoteControlProtocolOSC::ProcessAutoBinding(const FOSCAddress& InAddress)
 				{
 					Unbind(ProtocolEntityPtr);
 
-					FRemoteControlOSCProtocolEntity* OSCProtocolEntity = ProtocolEntityPtr->CastChecked<FRemoteControlOSCProtocolEntity>();
-					OSCProtocolEntity->PathName = *InAddress.GetFullPath();
+					if (FRemoteControlOSCProtocolEntity* OSCProtocolEntity = ProtocolEntityPtr->Cast<FRemoteControlOSCProtocolEntity>())
+					{
+						OSCProtocolEntity->PathName = *InAddress.GetFullPath();						
+					}
 
 					Bind(ProtocolEntityPtr);
 				}
